Add runtime physx cooking

diff --git a/Engine/Source/Developer/PhysXFormats/PhysXFormats.Build.cs b/Engine/Source/Developer/PhysXFormats/PhysXFormats.Build.cs
deleted file mode 100644
diff --git a/Engine/Source/Developer/PhysXFormats/Private/PhysXFormats.cpp b/Engine/Source/Developer/PhysXFormats/Private/PhysXFormats.cpp
deleted file mode 100644
diff --git a/Engine/Source/Developer/PhysXFormats/Private/PhysXFormats.h b/Engine/Source/Developer/PhysXFormats/Private/PhysXFormats.h
deleted file mode 100644
diff --git a/Engine/Source/Developer/TargetPlatform/Private/TargetPlatformManagerModule.cpp b/Engine/Source/Developer/TargetPlatform/Private/TargetPlatformManagerModule.cpp
--- a/Engine/Source/Developer/TargetPlatform/Private/TargetPlatformManagerModule.cpp
+++ b/Engine/Source/Developer/TargetPlatform/Private/TargetPlatformManagerModule.cpp
@@ -7,6 +7,8 @@
 #include "TargetPlatformPrivatePCH.h"
 #include "PlatformInfo.h"
 #include "DesktopPlatformModule.h"
+#include "IPhysXFormatModule.h"
+#include "IPhysXFormat.h"
 
 DEFINE_LOG_CATEGORY_STATIC(LogTargetPlatformManager, Log, All);
 
diff --git a/Engine/Source/Developer/TargetPlatform/Public/Interfaces/IPhysXFormat.h b/Engine/Source/Developer/TargetPlatform/Public/Interfaces/IPhysXFormat.h
deleted file mode 100644
diff --git a/Engine/Source/Developer/TargetPlatform/Public/Interfaces/IPhysXFormatModule.h b/Engine/Source/Developer/TargetPlatform/Public/Interfaces/IPhysXFormatModule.h
deleted file mode 100644
diff --git a/Engine/Source/Developer/TargetPlatform/Public/TargetPlatform.h b/Engine/Source/Developer/TargetPlatform/Public/TargetPlatform.h
--- a/Engine/Source/Developer/TargetPlatform/Public/TargetPlatform.h
+++ b/Engine/Source/Developer/TargetPlatform/Public/TargetPlatform.h
@@ -26,9 +26,6 @@
 #include "IAudioFormat.h"
 #include "IAudioFormatModule.h"
 
-#include "IPhysXFormat.h"
-#include "IPhysXFormatModule.h"
-
 #include "IShaderFormat.h"
 #include "IShaderFormatModule.h"
 
diff --git a/Engine/Source/Developer/TargetPlatform/TargetPlatform.Build.cs b/Engine/Source/Developer/TargetPlatform/TargetPlatform.Build.cs
--- a/Engine/Source/Developer/TargetPlatform/TargetPlatform.Build.cs
+++ b/Engine/Source/Developer/TargetPlatform/TargetPlatform.Build.cs
@@ -10,6 +10,8 @@ public class TargetPlatform : ModuleRules
 		PrivateDependencyModuleNames.Add("Core");		
         PublicDependencyModuleNames.Add("DesktopPlatform");
 
+        PrivateIncludePathModuleNames.Add("PhysXFormats");
+
 		// no need for all these modules if the program doesn't want developer tools at all (like UnrealFileServer)
 		if (!UEBuildConfiguration.bBuildRequiresCookedData && UEBuildConfiguration.bBuildDeveloperTools)
 		{
@@ -102,11 +104,11 @@ public class TargetPlatform : ModuleRules
                     DynamicallyLoadedModuleNames.Add("AudioFormatOpus");
                 }
             }
-
-			if (UEBuildConfiguration.bCompileAgainstEngine && UEBuildConfiguration.bCompilePhysX)
-			{
-				DynamicallyLoadedModuleNames.Add("PhysXFormats");
-			}
 		}
+        
+        if (UEBuildConfiguration.bBuildDeveloperTools == true && (UEBuildConfiguration.bBuildRequiresCookedData || UEBuildConfiguration.bRuntimePhysicsCooking) && UEBuildConfiguration.bCompileAgainstEngine && UEBuildConfiguration.bCompilePhysX)
+        {
+            DynamicallyLoadedModuleNames.Add("PhysXFormats");
+        }
 	}
 }
diff --git a/Engine/Source/Programs/UnrealBuildTool/Android/UEBuildAndroid.cs b/Engine/Source/Programs/UnrealBuildTool/Android/UEBuildAndroid.cs
--- a/Engine/Source/Programs/UnrealBuildTool/Android/UEBuildAndroid.cs
+++ b/Engine/Source/Programs/UnrealBuildTool/Android/UEBuildAndroid.cs
@@ -175,6 +175,7 @@ namespace UnrealBuildTool
 			UEBuildConfiguration.bCompileLeanAndMeanUE = true;
 			UEBuildConfiguration.bCompilePhysX = true;
 			UEBuildConfiguration.bCompileAPEX = false;
+            UEBuildConfiguration.bRuntimePhysicsCooking = false;
 
 			UEBuildConfiguration.bBuildEditor = false;
 			UEBuildConfiguration.bBuildDeveloperTools = false;
diff --git a/Engine/Source/Programs/UnrealBuildTool/Configuration/UEBuildConfiguration.cs b/Engine/Source/Programs/UnrealBuildTool/Configuration/UEBuildConfiguration.cs
--- a/Engine/Source/Programs/UnrealBuildTool/Configuration/UEBuildConfiguration.cs
+++ b/Engine/Source/Programs/UnrealBuildTool/Configuration/UEBuildConfiguration.cs
@@ -15,6 +15,10 @@ namespace UnrealBuildTool
 		[XmlConfig]
 		public static bool bCompileAPEX;
 
+        /** Whether to allow runtime cooking of physics */
+        [XmlConfig]
+        public static bool bRuntimePhysicsCooking;
+
 		/** Whether to include Box2D support */
 		[XmlConfig]
 		public static bool bCompileBox2D;
@@ -188,6 +192,7 @@ namespace UnrealBuildTool
 			// Currently, WITH_PHYSX is forced to true in Engine.h (as it isn't defined anywhere by the builder)
 			bCompilePhysX = true;
 			bCompileAPEX = true;
+            bRuntimePhysicsCooking = true;
 			bCompileBox2D = true;
 			bBuildDedicatedServer = false;
 			bBuildEditor = true;
diff --git a/Engine/Source/Programs/UnrealBuildTool/HTML5/UEBuildHTML5.cs b/Engine/Source/Programs/UnrealBuildTool/HTML5/UEBuildHTML5.cs
--- a/Engine/Source/Programs/UnrealBuildTool/HTML5/UEBuildHTML5.cs
+++ b/Engine/Source/Programs/UnrealBuildTool/HTML5/UEBuildHTML5.cs
@@ -269,6 +269,7 @@ namespace UnrealBuildTool
             UEBuildConfiguration.bCompileLeanAndMeanUE = true;
             UEBuildConfiguration.bCompileAPEX = false;
             UEBuildConfiguration.bCompilePhysX = true;
+            UEBuildConfiguration.bRuntimePhysicsCooking = false;
             UEBuildConfiguration.bCompileSimplygon = false;
             UEBuildConfiguration.bCompileICU = false;
         }
diff --git a/Engine/Source/Programs/UnrealBuildTool/IOS/UEBuildIOS.cs b/Engine/Source/Programs/UnrealBuildTool/IOS/UEBuildIOS.cs
--- a/Engine/Source/Programs/UnrealBuildTool/IOS/UEBuildIOS.cs
+++ b/Engine/Source/Programs/UnrealBuildTool/IOS/UEBuildIOS.cs
@@ -185,6 +185,7 @@ namespace UnrealBuildTool
             UEBuildConfiguration.bBuildEditor = false;
             UEBuildConfiguration.bBuildDeveloperTools = false;
             UEBuildConfiguration.bCompileAPEX = false;
+            UEBuildConfiguration.bRuntimePhysicsCooking = false;
             UEBuildConfiguration.bCompileSimplygon = false;
             UEBuildConfiguration.bBuildDeveloperTools = false;
             UEBuildConfiguration.bCompileICU = true;
diff --git a/Engine/Source/Programs/UnrealBuildTool/System/RulesCompiler.cs b/Engine/Source/Programs/UnrealBuildTool/System/RulesCompiler.cs
--- a/Engine/Source/Programs/UnrealBuildTool/System/RulesCompiler.cs
+++ b/Engine/Source/Programs/UnrealBuildTool/System/RulesCompiler.cs
@@ -279,6 +279,11 @@ namespace UnrealBuildTool
 				Definitions.Add("WITH_PHYSX=0");
 				Definitions.Add("WITH_APEX=0");
 			}
+
+            if(UEBuildConfiguration.bRuntimePhysicsCooking == true)
+            {
+                Definitions.Add("WITH_RUNTIME_PHYSICS_COOKING");
+            }
 		}
 
 		/// <summary>
diff --git a/Engine/Source/Programs/UnrealBuildTool/WinRT/UEBuildWinRT.cs b/Engine/Source/Programs/UnrealBuildTool/WinRT/UEBuildWinRT.cs
--- a/Engine/Source/Programs/UnrealBuildTool/WinRT/UEBuildWinRT.cs
+++ b/Engine/Source/Programs/UnrealBuildTool/WinRT/UEBuildWinRT.cs
@@ -171,6 +171,7 @@ namespace UnrealBuildTool
             UEBuildConfiguration.bCompileLeanAndMeanUE = true;
             UEBuildConfiguration.bCompilePhysX = false;
             UEBuildConfiguration.bCompileAPEX = false;
+            UEBuildConfiguration.bRuntimePhysicsCooking = false;
 
             BuildConfiguration.bUseUnityBuild = false;
 
diff --git a/Engine/Source/Runtime/Engine/Classes/PhysicsEngine/BodySetup.h b/Engine/Source/Runtime/Engine/Classes/PhysicsEngine/BodySetup.h
--- a/Engine/Source/Runtime/Engine/Classes/PhysicsEngine/BodySetup.h
+++ b/Engine/Source/Runtime/Engine/Classes/PhysicsEngine/BodySetup.h
@@ -541,7 +541,7 @@ public:
 	/** Returns the physics material used for this body. If none, specified, returns the default engine material. */
 	ENGINE_API class UPhysicalMaterial* GetPhysMaterial() const;
 
-#if WITH_EDITOR
+#if WITH_RUNTIME_PHYSICS_COOKING || WITH_EDITOR
 	/** Clear all simple collision */
 	ENGINE_API void RemoveSimpleCollision();
 
@@ -566,7 +566,7 @@ public:
 	 */
 	ENGINE_API void CreateFromModel(class UModel* InModel, bool bRemoveExisting);
 
-#endif // WITH_EDITOR
+#endif // WITH_RUNTIME_PHYSICS_COOKING || WITH_EDITOR
 
 	/**
 	 * Converts the skinned data of a skeletal mesh into a tri mesh collision. This is used for per poly scene queries and is quite expensive.
diff --git a/Engine/Source/Runtime/Engine/Engine.Build.cs b/Engine/Source/Runtime/Engine/Engine.Build.cs
--- a/Engine/Source/Runtime/Engine/Engine.Build.cs
+++ b/Engine/Source/Runtime/Engine/Engine.Build.cs
@@ -194,6 +194,12 @@ public class Engine : ModuleRules
 		}
 
 		SetupModulePhysXAPEXSupport(Target);
+        if(UEBuildConfiguration.bCompilePhysX && UEBuildConfiguration.bRuntimePhysicsCooking)
+        {
+            DynamicallyLoadedModuleNames.Add("PhysXFormats");
+            PrivateIncludePathModuleNames.Add("PhysXFormats");
+        }
+            
 
 		SetupModuleBox2DSupport(Target);
 
diff --git a/Engine/Source/Runtime/Engine/Private/PhysicsEngine/BodySetup.cpp b/Engine/Source/Runtime/Engine/Private/PhysicsEngine/BodySetup.cpp
--- a/Engine/Source/Runtime/Engine/Private/PhysicsEngine/BodySetup.cpp
+++ b/Engine/Source/Runtime/Engine/Private/PhysicsEngine/BodySetup.cpp
@@ -292,7 +292,7 @@ void SetupNonUniformHelper(FVector & Scale3D, float & MinScale, float & MinScale
 
 void UBodySetup::AddShapesToRigidActor(PxRigidActor* PDestActor, FVector& Scale3D, const FTransform * RelativeTM /* = NULL */, TArray<physx::PxShape*> * NewShapes /* = NULL */ )
 {
-#if WITH_EDITOR
+#if WITH_RUNTIME_PHYSICS_COOKING || WITH_EDITOR
 	// in editor, there are a lot of things relying on body setup to create physics meshes
 	CreatePhysicsMeshes();
 #endif
@@ -556,7 +556,7 @@ void UBodySetup::AddShapesToRigidActor(PxRigidActor* PDestActor, FVector& Scale3
 #endif // WITH_PHYSX
 
 
-#if WITH_EDITOR
+#if WITH_RUNTIME_PHYSICS_COOKING || WITH_EDITOR
 
 void UBodySetup::RemoveSimpleCollision()
 {
@@ -828,12 +828,16 @@ FByteBulkData* UBodySetup::GetCookedData(FName Format)
 			return NULL;
 		}
 
-#if WITH_EDITOR
+#if WITH_RUNTIME_PHYSICS_COOKING || WITH_EDITOR
 		TArray<uint8> OutData;
 		FDerivedDataPhysXCooker* DerivedPhysXData = new FDerivedDataPhysXCooker(Format, this);
 		if (DerivedPhysXData->CanBuild())
 		{
+		#if WITH_EDITOR
 			GetDerivedDataCacheRef().GetSynchronous(DerivedPhysXData, OutData);
+		#elif WITH_RUNTIME_PHYSICS_COOKING
+			DerivedPhysXData->Build(OutData);
+		#endif
 			if (OutData.Num())
 			{
 				Result->Lock(LOCK_READ_WRITE);
diff --git a/Engine/Source/Runtime/Engine/Private/PhysicsEngine/PhysDerivedData.cpp b/Engine/Source/Runtime/Engine/Private/PhysicsEngine/PhysDerivedData.cpp
--- a/Engine/Source/Runtime/Engine/Private/PhysicsEngine/PhysDerivedData.cpp
+++ b/Engine/Source/Runtime/Engine/Private/PhysicsEngine/PhysDerivedData.cpp
@@ -3,7 +3,8 @@
 #include "PhysDerivedData.h"
 #include "TargetPlatform.h"
 
-#if WITH_PHYSX && WITH_EDITOR
+#if WITH_PHYSX && (WITH_RUNTIME_PHYSICS_COOKING || WITH_EDITOR)
+#include "IPhysxFormatModule.h"
 #include "Landscape/LandscapeMeshCollisionComponent.h"
 
 FDerivedDataPhysXCooker::FDerivedDataPhysXCooker( FName InFormat, UBodySetup* InBodySetup )
@@ -41,12 +42,19 @@ FDerivedDataPhysXCooker::FDerivedDataPhysXCooker( FName InFormat, ULandscapeMesh
 
 void FDerivedDataPhysXCooker::InitCooker()
 {
+#if WITH_EDITOR
 	// static here as an optimization
 	static ITargetPlatformManagerModule* TPM = GetTargetPlatformManager();
 	if (TPM)
 	{
 		Cooker = TPM->FindPhysXFormat( Format );
 	}
+#elif WITH_RUNTIME_PHYSICS_COOKING
+	if (IPhysXFormatModule* Module = FModuleManager::LoadModulePtr<IPhysXFormatModule>("PhysXFormats"))
+	{
+		Cooker = Module->GetPhysXFormat();
+	}
+#endif
 }
 
 bool FDerivedDataPhysXCooker::Build( TArray<uint8>& OutData )
diff --git a/Engine/Source/Runtime/Engine/Private/PhysicsEngine/PhysDerivedData.h b/Engine/Source/Runtime/Engine/Private/PhysicsEngine/PhysDerivedData.h
--- a/Engine/Source/Runtime/Engine/Private/PhysicsEngine/PhysDerivedData.h
+++ b/Engine/Source/Runtime/Engine/Private/PhysicsEngine/PhysDerivedData.h
@@ -2,12 +2,13 @@
 
 #pragma once
 
-#if WITH_PHYSX && WITH_EDITOR
+#if WITH_PHYSX && (WITH_RUNTIME_PHYSICS_COOKING || WITH_EDITOR)
 
 #include "DerivedDataPluginInterface.h"
 #include "DerivedDataCacheInterface.h"
 #include "TargetPlatform.h"
 #include "PhysXSupport.h"
+#include "IPhysXFormat.h"
 
 // Forward declarations
 class ULandscapeMeshCollisionComponent;
diff --git a/Engine/Source/Runtime/Engine/Private/PhysicsEngine/PhysUtils.cpp b/Engine/Source/Runtime/Engine/Private/PhysicsEngine/PhysUtils.cpp
--- a/Engine/Source/Runtime/Engine/Private/PhysicsEngine/PhysUtils.cpp
+++ b/Engine/Source/Runtime/Engine/Private/PhysicsEngine/PhysUtils.cpp
@@ -109,7 +109,7 @@ static bool ModelToHullsWorker(FKAggregateGeom* outGeom,
 	return true;
 }
 
-#if WITH_EDITOR
+#if WITH_RUNTIME_PHYSICS_COOKING || WITH_EDITOR
 void UBodySetup::CreateFromModel(UModel* InModel, bool bRemoveExisting)
 {
 	if ( bRemoveExisting )
@@ -133,7 +133,7 @@ void UBodySetup::CreateFromModel(UModel* InModel, bool bRemoveExisting)
 	// Create new GUID
 	InvalidatePhysicsData();
 }
-#endif // WITH_EDITOR
+#endif // WITH_RUNTIME_PHYSICS_COOKING || WITH_EDITOR
 
 //////////////////////////////////////////////////////////////////////////
 // FRigidBodyCollisionInfo
diff --git a/Engine/Source/Runtime/Engine/Public/EngineDefines.h b/Engine/Source/Runtime/Engine/Public/EngineDefines.h
--- a/Engine/Source/Runtime/Engine/Public/EngineDefines.h
+++ b/Engine/Source/Runtime/Engine/Public/EngineDefines.h
@@ -54,10 +54,6 @@
 #endif//WITH_CLOTH_COLLISION_DETECTION
 #endif //WITH_APEX_CLOTHING
 
-#ifndef WITH_BODY_WELDING
-	//#define WITH_BODY_WELDING 1
-#endif
-
 #ifndef ENABLE_VISUAL_LOG
 	#define ENABLE_VISUAL_LOG (1 && !NO_LOGGING && !USING_CODE_ANALYSIS && !(UE_BUILD_SHIPPING || UE_BUILD_TEST))
 #endif
diff --git a/Engine/Source/Runtime/PhysXFormats/PhysXFormats.Build.cs b/Engine/Source/Runtime/PhysXFormats/PhysXFormats.Build.cs
new file mode 100644
--- /dev/null
+++ b/Engine/Source/Runtime/PhysXFormats/PhysXFormats.Build.cs
@@ -0,0 +1,22 @@
+// Copyright 1998-2014 Epic Games, Inc. All Rights Reserved.
+
+using UnrealBuildTool;
+using System.IO;
+
+public class PhysXFormats : ModuleRules
+{
+	public PhysXFormats(TargetInfo Target)
+	{
+		PrivateIncludePathModuleNames.Add("TargetPlatform");
+
+		PrivateDependencyModuleNames.AddRange(
+			new string[] {
+				"Core",
+				"CoreUObject", // @todo Mac: for some reason it's needed to link in debug on Mac
+				"Engine"
+			}
+			);
+
+		SetupModulePhysXAPEXSupport(Target);
+	}
+}
diff --git a/Engine/Source/Runtime/PhysXFormats/Private/PhysXFormats.cpp b/Engine/Source/Runtime/PhysXFormats/Private/PhysXFormats.cpp
new file mode 100644
--- /dev/null
+++ b/Engine/Source/Runtime/PhysXFormats/Private/PhysXFormats.cpp
@@ -0,0 +1,257 @@
+// Copyright 1998-2014 Epic Games, Inc. All Rights Reserved.
+
+#include "Core.h"
+#include "ModuleInterface.h"
+#include "ModuleManager.h"
+#include "Engine.h"
+#include "PhysicsPublic.h"
+#include "TargetPlatform.h"
+#include "IPhysXFormat.h"
+#include "IPhysXFormatModule.h"
+#include "PhysXFormats.h"
+
+static_assert(WITH_PHYSX, "No point in compiling PhysX cooker, if we don't have PhysX.");
+
+static FName NAME_PhysXPC(TEXT("PhysXPC"));
+static FName NAME_PhysXXboxOne(TEXT("PhysXXboxOne"));
+static FName NAME_PhysXPS4(TEXT("PhysXPS4"));
+
+/**
+ * FPhysXFormats. Cooks physics data.
+**/
+class FPhysXFormats : public IPhysXFormat
+{
+	enum
+	{
+		/** Version for PhysX format, this becomes part of the DDC key. */
+		UE_PHYSX_PC_VER = 0,
+	};
+
+	PxCooking* PhysXCooking;
+
+	/**
+	 * Validates format name and returns its PhysX enum value.
+	 *
+	 * @param InFormatName PhysX format name.
+	 * @param OutFormat PhysX enum
+	 * @return true if InFormatName is a valid and supported PhysX format
+	 */
+	bool GetPhysXFormat(FName InFormatName, PxPlatform::Enum& OutFormat) const
+	{
+		if ((InFormatName == NAME_PhysXPC) || (InFormatName == NAME_PhysXXboxOne) || (InFormatName == NAME_PhysXPS4))
+		{
+			OutFormat = PxPlatform::ePC;
+		}
+		else
+		{
+			return false;
+		}
+		return true;
+	}
+
+	/**
+	 * Valudates PhysX format name.
+	 */
+	bool CheckPhysXFormat(FName InFormatName) const
+	{
+		PxPlatform::Enum PhysXFormat = PxPlatform::ePC;
+		return GetPhysXFormat(InFormatName, PhysXFormat);
+	}
+
+public:
+
+	FPhysXFormats( PxCooking* InCooking )
+		: PhysXCooking( InCooking )
+	{}
+
+	virtual bool AllowParallelBuild() const
+	{
+		return false;
+	}
+
+	virtual uint16 GetVersion(FName Format) const override
+	{
+		check(CheckPhysXFormat(Format));
+		return UE_PHYSX_PC_VER;
+	}
+
+
+	virtual void GetSupportedFormats(TArray<FName>& OutFormats) const
+	{
+		OutFormats.Add(NAME_PhysXPC);
+		OutFormats.Add(NAME_PhysXXboxOne);
+		OutFormats.Add(NAME_PhysXPS4);
+	}
+
+	virtual bool CookConvex(FName Format, const TArray<FVector>& SrcBuffer, TArray<uint8>& OutBuffer) const override
+	{
+#if WITH_PHYSX
+		PxPlatform::Enum PhysXFormat = PxPlatform::ePC;
+		bool bIsPhysXFormatValid = GetPhysXFormat(Format, PhysXFormat);
+		check(bIsPhysXFormatValid);
+
+		PxConvexMeshDesc PConvexMeshDesc;
+		PConvexMeshDesc.points.data = SrcBuffer.GetData();
+		PConvexMeshDesc.points.count = SrcBuffer.Num();
+		PConvexMeshDesc.points.stride = sizeof(FVector);
+		PConvexMeshDesc.flags = PxConvexFlag::eCOMPUTE_CONVEX;
+
+		// Set up cooking
+		const PxCookingParams& Params = PhysXCooking->getParams();
+		PxCookingParams NewParams = Params;
+		NewParams.targetPlatform = PhysXFormat;
+		PhysXCooking->setParams(NewParams);
+
+		// Cook the convex mesh to a temp buffer
+		TArray<uint8> CookedMeshBuffer;
+		FPhysXOutputStream Buffer(&CookedMeshBuffer);
+		bool Result = PhysXCooking->cookConvexMesh(PConvexMeshDesc, Buffer);
+		
+		if( Result && CookedMeshBuffer.Num() > 0 )
+		{
+			// Append the cooked data into cooked buffer
+			OutBuffer.Append( CookedMeshBuffer );
+			return true;
+		}
+#endif		// WITH_PHYSX
+		return false;
+	}
+
+	virtual bool CookTriMesh(FName Format, const TArray<FVector>& SrcVertices, const TArray<FTriIndices>& SrcIndices, const TArray<uint16>& SrcMaterialIndices, const bool FlipNormals, TArray<uint8>& OutBuffer, bool bPerPolySkeletalMesh = false) const override
+	{
+#if WITH_PHYSX
+		PxPlatform::Enum PhysXFormat = PxPlatform::ePC;
+		bool bIsPhysXFormatValid = GetPhysXFormat(Format, PhysXFormat);
+		check(bIsPhysXFormatValid);
+
+		PxTriangleMeshDesc PTriMeshDesc;
+		PTriMeshDesc.points.data = SrcVertices.GetData();
+		PTriMeshDesc.points.count = SrcVertices.Num();
+		PTriMeshDesc.points.stride = sizeof(FVector);
+		PTriMeshDesc.triangles.data = SrcIndices.GetData();
+		PTriMeshDesc.triangles.count = SrcIndices.Num();
+		PTriMeshDesc.triangles.stride = sizeof(FTriIndices);
+		PTriMeshDesc.materialIndices.data = SrcMaterialIndices.GetData();
+		PTriMeshDesc.materialIndices.stride = sizeof(PxMaterialTableIndex);
+		PTriMeshDesc.flags = FlipNormals ? PxMeshFlag::eFLIPNORMALS : (PxMeshFlags)0;
+
+		// Set up cooking
+		const PxCookingParams& Params = PhysXCooking->getParams();
+		PxCookingParams NewParams = Params;
+		NewParams.targetPlatform = PhysXFormat;
+		PxMeshPreprocessingFlags OldCookingFlags = NewParams.meshPreprocessParams;
+
+		if (bPerPolySkeletalMesh)
+		{
+			//per poly skeletal mesh requires deforming mesh. This is a very special case so we have to change the cook params
+			NewParams.meshPreprocessParams = PxMeshPreprocessingFlag::eDISABLE_CLEAN_MESH;
+		}
+
+		PhysXCooking->setParams(NewParams);
+
+
+		// Cook TriMesh Data
+		FPhysXOutputStream Buffer(&OutBuffer);
+		bool Result = PhysXCooking->cookTriangleMesh(PTriMeshDesc, Buffer);
+		
+		if (bPerPolySkeletalMesh)	//restore old params
+		{
+			NewParams.meshPreprocessParams = OldCookingFlags;
+			PhysXCooking->setParams(NewParams);
+		}
+		return Result;
+#else
+		return false;
+#endif		// WITH_PHYSX
+	}
+
+	virtual bool CookHeightField(FName Format, FIntPoint HFSize, float Thickness, const void* Samples, uint32 SamplesStride, TArray<uint8>& OutBuffer) const override
+	{
+#if WITH_PHYSX
+		PxPlatform::Enum PhysXFormat = PxPlatform::ePC;
+		bool bIsPhysXFormatValid = GetPhysXFormat(Format, PhysXFormat);
+		check(bIsPhysXFormatValid);
+
+		PxHeightFieldDesc HFDesc;
+		HFDesc.format = PxHeightFieldFormat::eS16_TM;
+		HFDesc.nbColumns = HFSize.X;
+		HFDesc.nbRows = HFSize.Y;
+		HFDesc.samples.data = Samples;
+		HFDesc.samples.stride = SamplesStride;
+		HFDesc.flags = PxHeightFieldFlag::eNO_BOUNDARY_EDGES;
+		HFDesc.thickness = Thickness;
+
+		// Set up cooking
+		const PxCookingParams& Params = PhysXCooking->getParams();
+		PxCookingParams NewParams = Params;
+		NewParams.targetPlatform = PhysXFormat;
+		PhysXCooking->setParams(NewParams);
+
+		// Cook to a temp buffer
+		TArray<uint8> CookedBuffer;
+		FPhysXOutputStream Buffer(&CookedBuffer);
+
+		if (PhysXCooking->cookHeightField(HFDesc, Buffer) && CookedBuffer.Num() > 0)
+		{
+			// Append the cooked data into cooked buffer
+			OutBuffer.Append(CookedBuffer);
+			return true;
+		}
+		return false;
+#else
+		return false;
+#endif		// WITH_PHYSX
+	}
+
+};
+
+
+/**
+ * Module for PhysX cooking
+ */
+static IPhysXFormat* Singleton = NULL;
+
+class FPhysXPlatformModule : public IPhysXFormatModule
+{
+public:
+	FPhysXPlatformModule()
+	{
+
+	}
+
+	virtual ~FPhysXPlatformModule()
+	{
+		ShutdownPhysXCooking();
+
+		delete Singleton;
+		Singleton = NULL;
+	}
+	virtual IPhysXFormat* GetPhysXFormat()
+	{
+		if (!Singleton)
+		{
+			InitPhysXCooking();
+			Singleton = new FPhysXFormats(GPhysXCooking);
+		}
+		return Singleton;
+	}
+
+private:
+
+	/**
+	 *	Load the required modules for PhysX
+	 */
+	void InitPhysXCooking()
+	{
+		// Make sure PhysX libs are loaded
+		LoadPhysXModules();
+	}
+
+	void ShutdownPhysXCooking()
+	{
+		// Ideally PhysX cooking should be initialized in InitPhysXCooking and released here
+		// but Apex is still being setup in the engine and it also requires PhysX cooking singleton.
+	}
+};
+
+IMPLEMENT_MODULE( FPhysXPlatformModule, PhysXFormats );
diff --git a/Engine/Source/Runtime/PhysXFormats/Private/PhysXFormats.h b/Engine/Source/Runtime/PhysXFormats/Private/PhysXFormats.h
new file mode 100644
--- /dev/null
+++ b/Engine/Source/Runtime/PhysXFormats/Private/PhysXFormats.h
@@ -0,0 +1,82 @@
+// Copyright 1998-2014 Epic Games, Inc. All Rights Reserved.
+
+#pragma once
+
+#if WITH_PHYSX
+
+#if _MSC_VER
+#pragma warning( push )
+#pragma warning( disable : 4946 ) // reinterpret_cast used between related classes
+#endif
+
+#ifdef __clang__
+#pragma clang diagnostic push
+#pragma clang diagnostic ignored "-Wshadow"
+#endif
+
+#if USING_CODE_ANALYSIS
+#pragma warning( push )
+#pragma warning( disable : ALL_CODE_ANALYSIS_WARNINGS )
+#endif	// USING_CODE_ANALYSIS
+
+#pragma pack(push,8)
+
+#include "Px.h"
+#include "PxPhysicsAPI.h"
+#include "PxRenderBuffer.h"
+#include "PxExtensionsAPI.h"
+#include "PxVisualDebuggerExt.h"
+//#include "PxDefaultCpuDispatcher.h"
+
+// vehicle related header files
+//#include "PxVehicleSDK.h"
+//#include "PxVehicleUtils.h"
+
+// utils
+#include "PxGeometryQuery.h"
+#include "PxMeshQuery.h"
+#include "PxTriangle.h"
+
+#pragma pack(pop)
+
+#if USING_CODE_ANALYSIS
+#pragma warning( pop )
+#endif	// USING_CODE_ANALYSIS
+
+#if _MSC_VER
+#pragma warning( pop )
+#endif
+
+#ifdef __clang__
+#pragma clang diagnostic pop
+#endif
+
+using namespace physx;
+
+/** Utility wrapper for a uint8 TArray for saving into PhysX. */
+class FPhysXOutputStream : public PxOutputStream
+{
+public:
+	/** Raw byte data */
+	TArray<uint8>			*Data;
+
+	FPhysXOutputStream()
+		: Data(NULL)
+	{}
+
+	FPhysXOutputStream(TArray<uint8> *InData)
+		: Data(InData)
+	{}
+
+	virtual PxU32 write(const void* Src, PxU32 Count) override
+	{
+		check(Data);
+		check(Src);
+		int32 CurrentNum = (*Data).Num();
+		(*Data).AddUninitialized(Count);
+		FMemory::Memcpy(&(*Data)[CurrentNum], Src, Count);
+		return Count;
+	}
+};
+
+#endif // WITH_PHYSX
diff --git a/Engine/Source/Runtime/PhysXFormats/Public/IPhysXFormat.h b/Engine/Source/Runtime/PhysXFormats/Public/IPhysXFormat.h
new file mode 100644
--- /dev/null
+++ b/Engine/Source/Runtime/PhysXFormats/Public/IPhysXFormat.h
@@ -0,0 +1,84 @@
+// Copyright 1998-2014 Epic Games, Inc. All Rights Reserved.
+
+/*=============================================================================
+	ITextureFormatModule.h: Declares the ITextureFormatModule interface.
+=============================================================================*/
+
+#pragma once
+
+
+/**
+ * IPhysXFormat, PhysX cooking abstraction
+**/
+class IPhysXFormat
+{
+public:
+
+	/**
+	 * Checks whether parallel PhysX cooking is allowed.
+	 *
+	 * Note: This method is not currently used yet.
+	 *
+	 * @return true if this PhysX format can cook in parallel, false otherwise.
+	 */
+	virtual bool AllowParallelBuild( ) const
+	{
+		return false;
+	}
+
+	/**
+	 * Cooks the source convex data for the platform and stores the cooked data internally.
+	 *
+	 * @param Format The desired format
+	 * @param SrcBuffer The source buffer
+	 * @param OutBuffer The resulting cooked data
+	 * @return true on success, false otherwise.
+	 */
+	virtual bool CookConvex( FName Format, const TArray<FVector>& SrcBuffer, TArray<uint8>& OutBuffer ) const = 0;
+
+	/**
+	 * Cooks the source Tri-Mesh data for the platform and stores the cooked data internally.
+	 *
+	 * @param Format The desired format.
+	 * @param SrcBuffer The source buffer.
+	 * @param OutBuffer The resulting cooked data.
+	 * @param bPerPolySkeletalMesh This is a very special case that requires different cooking parameters set.
+	 * @return true on success, false otherwise.
+	 */
+	virtual bool CookTriMesh( FName Format, const TArray<FVector>& SrcVertices, const TArray<struct FTriIndices>& SrcIndices, const TArray<uint16>& SrcMaterialIndices, const bool FlipNormals, TArray<uint8>& OutBuffer, bool bPerPolySkeletalMesh = false ) const = 0;
+		
+	/**
+	 * Cooks the source height field data for the platform and stores the cooked data internally.
+	 *
+	 * @param Format The desired format
+	 * @param HFSize Size of height field [NumColumns, NumRows]
+	 * @param Thickness Sets how thick the height field surface is
+	 * @param SrcBuffer The source buffer
+	 * @param OutBuffer The resulting cooked data
+	 * @return true on success, false otherwise.
+	 */
+	virtual bool CookHeightField( FName Format, FIntPoint HFSize, float Thickness, const void* Samples, uint32 SamplesStride, TArray<uint8>& OutBuffer ) const = 0;
+
+	/**
+	 * Gets the list of supported formats.
+	 *
+	 * @param OutFormats Will hold the list of formats.
+	 */
+	virtual void GetSupportedFormats( TArray<FName>& OutFormats ) const = 0;
+
+	/**
+	 * Gets the current version of the specified PhysX format.
+	 *
+	 * @param Format The format to get the version for.
+	 * @return Version number.
+	 */
+	virtual uint16 GetVersion( FName Format ) const = 0;
+
+
+public:
+
+	/**
+	 * Virtual destructor.
+	 */
+	virtual ~IPhysXFormat( ) { }
+};
diff --git a/Engine/Source/Runtime/PhysXFormats/Public/IPhysXFormatModule.h b/Engine/Source/Runtime/PhysXFormats/Public/IPhysXFormatModule.h
new file mode 100644
--- /dev/null
+++ b/Engine/Source/Runtime/PhysXFormats/Public/IPhysXFormatModule.h
@@ -0,0 +1,29 @@
+// Copyright 1998-2014 Epic Games, Inc. All Rights Reserved.
+
+/*=============================================================================
+	IPhysXFormatModule.h: Declares the IPhysXFormatModule interface.
+=============================================================================*/
+
+#pragma once
+
+
+/**
+ * Interface for PhysX format modules.
+ */
+class IPhysXFormatModule
+	: public IModuleInterface
+{
+public:
+
+	/**
+	 * Gets the PhysX format.
+	 */
+	virtual IPhysXFormat* GetPhysXFormat( ) = 0;
+
+public:
+
+	/**
+	 * Virtual destructor.
+	 */
+	~IPhysXFormatModule( ) { }
+};
