Fixes to SplineMeshComponent collision issues.

* Forces full collision rebuilds when modifying the mesh or geometry.
* Added extra logic to the cooking process when spline mesh simple collision is being cooked, so that it handles extreme deformity better.

diff --git a/Engine/Source/Runtime/Engine/Private/Components/SplineMeshComponent.cpp b/Engine/Source/Runtime/Engine/Private/Components/SplineMeshComponent.cpp
--- a/Engine/Source/Runtime/Engine/Private/Components/SplineMeshComponent.cpp
+++ b/Engine/Source/Runtime/Engine/Private/Components/SplineMeshComponent.cpp
@@ -486,7 +486,7 @@ void USplineMeshComponent::UpdateRenderStateAndCollision()
 #if WITH_EDITOR
 	if (!GetWorld()->AreActorsInitialized())
 	{
-		CachedMeshBodySetupGuid.Invalidate();
+		DestroyBodySetup();
 		RecreatePhysicsState();
 	}
 #endif // WITH_EDITOR
@@ -918,6 +918,7 @@ void USplineMeshComponent::DestroyBodySetup()
 	}
 }
 
+
 #if WITH_EDITOR
 void USplineMeshComponent::RecreateCollision()
 {
diff --git a/Engine/Source/Runtime/Engine/Private/EnginePrivate.h b/Engine/Source/Runtime/Engine/Private/EnginePrivate.h
--- a/Engine/Source/Runtime/Engine/Private/EnginePrivate.h
+++ b/Engine/Source/Runtime/Engine/Private/EnginePrivate.h
@@ -39,6 +39,7 @@ EnginePrivate.h: Unreal engine private header file.
 #include "Components/MeshComponent.h"
 #include "Components/SkinnedMeshComponent.h"
 #include "Components/SkeletalMeshComponent.h"
+#include "Components/SplineMeshComponent.h"
 #include "Components/AudioComponent.h"
 #include "Components/ArrowComponent.h"
 #include "Components/BillboardComponent.h"
diff --git a/Engine/Source/Runtime/Engine/Private/PhysicsEngine/PhysDerivedData.cpp b/Engine/Source/Runtime/Engine/Private/PhysicsEngine/PhysDerivedData.cpp
--- a/Engine/Source/Runtime/Engine/Private/PhysicsEngine/PhysDerivedData.cpp
+++ b/Engine/Source/Runtime/Engine/Private/PhysicsEngine/PhysDerivedData.cpp
@@ -123,7 +123,8 @@ int32 FDerivedDataPhysXCooker::BuildConvex( TArray<uint8>& OutData, bool InMirro
 
 		// Cook and store Result at ResultInfoOffset
 		UE_LOG(LogPhysics, Log, TEXT("Cook Convex: %s %d (FlipX:%d)"), *BodySetup->GetOuter()->GetPathName(), ElementIndex, InMirrored);		
-		bool Result = Cooker->CookConvex( Format, *MeshVertices, OutData );
+		const bool bDeformableMesh = CollisionDataProvider->IsA(USplineMeshComponent::StaticClass());
+		const bool Result = Cooker->CookConvex(Format, *MeshVertices, OutData, bDeformableMesh);
 		if( !Result )
 		{
 			UE_LOG(LogPhysics, Warning, TEXT("Failed to cook convex: %s %d (FlipX:%d). The remaining elements will not get cooked."), *BodySetup->GetOuter()->GetPathName(), ElementIndex, InMirrored);
@@ -169,13 +170,13 @@ int32 FDerivedDataPhysXCooker::BuildTriMesh( TArray<uint8>& OutData, bool InUseA
 		MeshVertices = &TriangleMeshDesc.Vertices;
 		
 		UE_LOG(LogPhysics, Log, TEXT("Cook TriMesh: %s"), *CollisionDataProvider->GetPathName());
-		bool bPerPolySkeletalMesh = false;
+		bool bDeformableMesh = CollisionDataProvider->IsA(USplineMeshComponent::StaticClass());
 		if (USkeletalMesh* SkeletalMesh = Cast<USkeletalMesh>(CollisionDataProvider))
 		{
 			ensure(SkeletalMesh->bEnablePerPolyCollision);
-			bPerPolySkeletalMesh = true;
+			bDeformableMesh = true;
 		}
-		bResult = Cooker->CookTriMesh( Format, *MeshVertices, TriangleMeshDesc.Indices, TriangleMeshDesc.MaterialIndices, TriangleMeshDesc.bFlipNormals, OutData, bPerPolySkeletalMesh );
+		bResult = Cooker->CookTriMesh( Format, *MeshVertices, TriangleMeshDesc.Indices, TriangleMeshDesc.MaterialIndices, TriangleMeshDesc.bFlipNormals, OutData, bDeformableMesh );
 		if( !bResult )
 		{
 			UE_LOG(LogPhysics, Warning, TEXT("Failed to cook TriMesh: %s."), *CollisionDataProvider->GetPathName());
diff --git a/Engine/Source/Runtime/PhysXFormats/Private/PhysXFormats.cpp b/Engine/Source/Runtime/PhysXFormats/Private/PhysXFormats.cpp
--- a/Engine/Source/Runtime/PhysXFormats/Private/PhysXFormats.cpp
+++ b/Engine/Source/Runtime/PhysXFormats/Private/PhysXFormats.cpp
@@ -84,7 +84,7 @@ public:
 		OutFormats.Add(NAME_PhysXPS4);
 	}
 
-	virtual bool CookConvex(FName Format, const TArray<FVector>& SrcBuffer, TArray<uint8>& OutBuffer) const override
+	virtual bool CookConvex(FName Format, const TArray<FVector>& SrcBuffer, TArray<uint8>& OutBuffer, bool bDeformableMesh = false) const override
 	{
 #if WITH_PHYSX
 		PxPlatform::Enum PhysXFormat = PxPlatform::ePC;
@@ -95,12 +95,28 @@ public:
 		PConvexMeshDesc.points.data = SrcBuffer.GetData();
 		PConvexMeshDesc.points.count = SrcBuffer.Num();
 		PConvexMeshDesc.points.stride = sizeof(FVector);
-		PConvexMeshDesc.flags = PxConvexFlag::eCOMPUTE_CONVEX;
+		if (bDeformableMesh)
+		{
+			PConvexMeshDesc.flags = PxConvexFlag::eCOMPUTE_CONVEX | PxConvexFlag::eINFLATE_CONVEX;
+		}
+		else
+		{
+			PConvexMeshDesc.flags = PxConvexFlag::eCOMPUTE_CONVEX;
+		}
 
 		// Set up cooking
 		const PxCookingParams& Params = PhysXCooking->getParams();
 		PxCookingParams NewParams = Params;
 		NewParams.targetPlatform = PhysXFormat;
+		if (bDeformableMesh)
+		{
+			// Meshes which can be deformed need different cooking parameters to inhibit vertex welding and add an extra skin around the collision mesh for safety.
+			// We need to set the meshWeldTolerance to zero, even when disabling 'clean mesh' as PhysX will attempt to perform mesh cleaning anyway according to this meshWeldTolerance
+			// if the convex hull is not well formed.
+			NewParams.meshPreprocessParams = PxMeshPreprocessingFlags(PxMeshPreprocessingFlag::eDISABLE_CLEAN_MESH | PxMeshPreprocessingFlag::eREMOVE_UNREFERENCED_VERTICES | PxMeshPreprocessingFlag::eREMOVE_DUPLICATED_TRIANGLES);
+			NewParams.areaTestEpsilon = 0.0f;
+			NewParams.meshWeldTolerance = 0.0f;
+		}
 		PhysXCooking->setParams(NewParams);
 
 		// Cook the convex mesh to a temp buffer
@@ -108,6 +124,14 @@ public:
 		FPhysXOutputStream Buffer(&CookedMeshBuffer);
 		bool Result = PhysXCooking->cookConvexMesh(PConvexMeshDesc, Buffer);
 		
+		// Return default cooking params to normal
+		if (bDeformableMesh)
+		{
+			NewParams.meshPreprocessParams = Params.meshPreprocessParams;
+			NewParams.areaTestEpsilon = Params.areaTestEpsilon;
+			NewParams.meshWeldTolerance = Params.meshWeldTolerance;
+		}
+
 		if( Result && CookedMeshBuffer.Num() > 0 )
 		{
 			// Append the cooked data into cooked buffer
@@ -118,7 +142,7 @@ public:
 		return false;
 	}
 
-	virtual bool CookTriMesh(FName Format, const TArray<FVector>& SrcVertices, const TArray<FTriIndices>& SrcIndices, const TArray<uint16>& SrcMaterialIndices, const bool FlipNormals, TArray<uint8>& OutBuffer, bool bPerPolySkeletalMesh = false) const override
+	virtual bool CookTriMesh(FName Format, const TArray<FVector>& SrcVertices, const TArray<FTriIndices>& SrcIndices, const TArray<uint16>& SrcMaterialIndices, const bool FlipNormals, TArray<uint8>& OutBuffer, bool bDeformableMesh = false) const override
 	{
 #if WITH_PHYSX
 		PxPlatform::Enum PhysXFormat = PxPlatform::ePC;
@@ -142,9 +166,9 @@ public:
 		NewParams.targetPlatform = PhysXFormat;
 		PxMeshPreprocessingFlags OldCookingFlags = NewParams.meshPreprocessParams;
 
-		if (bPerPolySkeletalMesh)
+		if (bDeformableMesh)
 		{
-			//per poly skeletal mesh requires deforming mesh. This is a very special case so we have to change the cook params
+			// In the case of a deformable mesh, we have to change the cook params
 			NewParams.meshPreprocessParams = PxMeshPreprocessingFlag::eDISABLE_CLEAN_MESH;
 		}
 
@@ -155,7 +179,7 @@ public:
 		FPhysXOutputStream Buffer(&OutBuffer);
 		bool Result = PhysXCooking->cookTriangleMesh(PTriMeshDesc, Buffer);
 		
-		if (bPerPolySkeletalMesh)	//restore old params
+		if (bDeformableMesh)	//restore old params
 		{
 			NewParams.meshPreprocessParams = OldCookingFlags;
 			PhysXCooking->setParams(NewParams);
diff --git a/Engine/Source/Runtime/PhysXFormats/Public/IPhysXFormat.h b/Engine/Source/Runtime/PhysXFormats/Public/IPhysXFormat.h
--- a/Engine/Source/Runtime/PhysXFormats/Public/IPhysXFormat.h
+++ b/Engine/Source/Runtime/PhysXFormats/Public/IPhysXFormat.h
@@ -32,9 +32,10 @@ public:
 	 * @param Format The desired format
 	 * @param SrcBuffer The source buffer
 	 * @param OutBuffer The resulting cooked data
+	 * @param bDeformableMesh Whether this mesh is deformable and hence needs to be cooked with different parameters
 	 * @return true on success, false otherwise.
 	 */
-	virtual bool CookConvex( FName Format, const TArray<FVector>& SrcBuffer, TArray<uint8>& OutBuffer ) const = 0;
+	virtual bool CookConvex( FName Format, const TArray<FVector>& SrcBuffer, TArray<uint8>& OutBuffer, bool bDeformableMesh = false ) const = 0;
 
 	/**
 	 * Cooks the source Tri-Mesh data for the platform and stores the cooked data internally.
@@ -42,10 +43,10 @@ public:
 	 * @param Format The desired format.
 	 * @param SrcBuffer The source buffer.
 	 * @param OutBuffer The resulting cooked data.
-	 * @param bPerPolySkeletalMesh This is a very special case that requires different cooking parameters set.
+	 * @param bDeformableMesh This is a very special case that requires different cooking parameters set.
 	 * @return true on success, false otherwise.
 	 */
-	virtual bool CookTriMesh( FName Format, const TArray<FVector>& SrcVertices, const TArray<struct FTriIndices>& SrcIndices, const TArray<uint16>& SrcMaterialIndices, const bool FlipNormals, TArray<uint8>& OutBuffer, bool bPerPolySkeletalMesh = false ) const = 0;
+	virtual bool CookTriMesh( FName Format, const TArray<FVector>& SrcVertices, const TArray<struct FTriIndices>& SrcIndices, const TArray<uint16>& SrcMaterialIndices, const bool FlipNormals, TArray<uint8>& OutBuffer, bool bDeformableMesh = false ) const = 0;
 		
 	/**
 	 * Cooks the source height field data for the platform and stores the cooked data internally.
