Missed spelling fixups

diff --git a/Engine/Source/Runtime/PhysXFormats/Private/PhysXFormats.cpp b/Engine/Source/Runtime/PhysXFormats/Private/PhysXFormats.cpp
--- a/Engine/Source/Runtime/PhysXFormats/Private/PhysXFormats.cpp
+++ b/Engine/Source/Runtime/PhysXFormats/Private/PhysXFormats.cpp
@@ -109,7 +109,7 @@ public:
 		PxCookingParams NewParams = Params;
 		NewParams.targetPlatform = PhysXFormat;
 
-		if(RuntimeCookFlags & ERuntimePhysxCookOptimizationFlags::SupressFaceRemapTable)
+		if(RuntimeCookFlags & ERuntimePhysxCookOptimizationFlags::SuppressFaceRemapTable)
 		{
 			NewParams.suppressTriangleMeshRemapTable = true;
 		}
@@ -182,7 +182,7 @@ public:
 		NewParams.targetPlatform = PhysXFormat;
 		PxMeshPreprocessingFlags OldCookingFlags = NewParams.meshPreprocessParams;
 
-		if (RuntimeCookFlags & ERuntimePhysxCookOptimizationFlags::SupressFaceRemapTable)
+		if (RuntimeCookFlags & ERuntimePhysxCookOptimizationFlags::SuppressFaceRemapTable)
 		{
 			NewParams.suppressTriangleMeshRemapTable = true;
 		}
